Electra: Fixed declaring modules with the same name.

diff --git a/Engine/Plugins/Media/AvidDNxHDDecoderElectra/Source/Private/AvidDNxHDDecoderElectraModule.cpp b/Engine/Plugins/Media/AvidDNxHDDecoderElectra/Source/Private/AvidDNxHDDecoderElectraModule.cpp
--- a/Engine/Plugins/Media/AvidDNxHDDecoderElectra/Source/Private/AvidDNxHDDecoderElectraModule.cpp
+++ b/Engine/Plugins/Media/AvidDNxHDDecoderElectra/Source/Private/AvidDNxHDDecoderElectraModule.cpp
@@ -35,6 +35,6 @@ public:
 	}
 };
 
-IMPLEMENT_MODULE(FAvidDNxHDElectraDecoderModule, HAPDecoderElectra);
+IMPLEMENT_MODULE(FAvidDNxHDElectraDecoderModule, AvidDNxHDDecoderElectra);
 
 #undef LOCTEXT_NAMESPACE
