Avoid warning V547: Expression is always false.
#rnx

diff --git a/Engine/Plugins/Media/AvidDNxHDDecoderElectra/Source/Private/Windows/ElectraMediaAvidDNxHDDecoder.cpp b/Engine/Plugins/Media/AvidDNxHDDecoderElectra/Source/Private/Windows/ElectraMediaAvidDNxHDDecoder.cpp
--- a/Engine/Plugins/Media/AvidDNxHDDecoderElectra/Source/Private/Windows/ElectraMediaAvidDNxHDDecoder.cpp
+++ b/Engine/Plugins/Media/AvidDNxHDDecoderElectra/Source/Private/Windows/ElectraMediaAvidDNxHDDecoder.cpp
@@ -819,12 +819,12 @@ bool FVideoDecoderAvidDNxHDElectra::FDecoderHandle::SetupDecompressStruct()
 		return false;
 	}
 
-	if (CurrentUncompressedParams.compType == DNX_ComponentType_t::DNX_CT_SHORT_2_14)
+	/*if (CurrentUncompressedParams.compType == DNX_ComponentType_t::DNX_CT_SHORT_2_14)
 	{
 		CurrentUncompressedParams.blackPoint = DNX_DEFAULT_SHORT_2_14_BLACK_POINT;
 		CurrentUncompressedParams.whitePoint = DNX_DEFAULT_SHORT_2_14_WHITE_POINT;
 		CurrentUncompressedParams.chromaExcursion = DNX_DEFAULT_CHROMA_SHORT_2_14_EXCURSION;
-	}
+	}*/
 	return true;
 }
 
