Disable DNxSDK on Windows Arm64

#rnx

diff --git a/Engine/Plugins/Media/AvidDNxHDDecoderElectra/AvidDNxHDDecoderElectra.uplugin b/Engine/Plugins/Media/AvidDNxHDDecoderElectra/AvidDNxHDDecoderElectra.uplugin
--- a/Engine/Plugins/Media/AvidDNxHDDecoderElectra/AvidDNxHDDecoderElectra.uplugin
+++ b/Engine/Plugins/Media/AvidDNxHDDecoderElectra/AvidDNxHDDecoderElectra.uplugin
@@ -18,6 +18,7 @@
 			"Type": "Runtime",
 			"LoadingPhase": "PostEngineInit",
 			"PlatformAllowList": [ "Win64" ],
+			"PlatformArchitectureDenyList": [ "Win64:arm64" ],
 			"TargetDenyList": [ "Server" ]
 		}
 	],
diff --git a/Engine/Plugins/Media/AvidDNxMedia/AvidDNxMedia.uplugin b/Engine/Plugins/Media/AvidDNxMedia/AvidDNxMedia.uplugin
--- a/Engine/Plugins/Media/AvidDNxMedia/AvidDNxMedia.uplugin
+++ b/Engine/Plugins/Media/AvidDNxMedia/AvidDNxMedia.uplugin
@@ -17,6 +17,7 @@
 			"Type": "Runtime",
 			"LoadingPhase": "PostEngineInit",
 			"PlatformAllowList": [ "Win64" ],
+			"PlatformArchitectureDenyList": [ "Win64:arm64" ],
 			"TargetDenyList": [ "Server" ]
 		}
 	],
